Slate - Updating the standalone D3D and OpenGL Slate renderers to do a similar sRGB packing in the vertex buffer color, as well as a sRGB standards compliant transform in the pixel shader instead of the pow(color, 2.2) gamma conversion it was previously performing.  Since Photoshop does not display in the Mac's native gamma space of 2.2, and instead sticks to the sRGB standard, I've made it so the standalone OpenGL renderer also uses the sRGB standard.  Probably going to do the same in the RHI version as it doesn't make a ton of sense to vary it across platforms for Mac when Photoshop doesn't vary.

#pandacake

diff --git a/Engine/Shaders/StandaloneRenderer/D3D/GammaCorrectionCommon.hlsl b/Engine/Shaders/StandaloneRenderer/D3D/GammaCorrectionCommon.hlsl
new file mode 100644
--- /dev/null
+++ b/Engine/Shaders/StandaloneRenderer/D3D/GammaCorrectionCommon.hlsl
@@ -0,0 +1,59 @@
+// Copyright 1998-2015 Epic Games, Inc. All Rights Reserved.
+
+
+half3 LinearTo709Branchless(half3 lin) 
+{
+	lin = max(6.10352e-5, lin); // minimum positive non-denormal (fixes black problem on DX11 AMD and NV)
+	return min(lin * 4.5, pow(max(lin, 0.018), 0.45) * 1.099 - 0.099);
+}
+
+half3 LinearToSrgbBranchless(half3 lin) 
+{
+	lin = max(6.10352e-5, lin); // minimum positive non-denormal (fixes black problem on DX11 AMD and NV)
+	return min(lin * 12.92, pow(max(lin, 0.00313067), 1.0/2.4) * 1.055 - 0.055);
+	// Possible that mobile GPUs might have native pow() function?
+	//return min(lin * 12.92, exp2(log2(max(lin, 0.00313067)) * (1.0/2.4) + log2(1.055)) - 0.055);
+}
+
+half LinearToSrgbBranchingChannel(half lin) 
+{
+	if(lin < 0.00313067) return lin * 12.92;
+	return pow(lin, (1.0/2.4)) * 1.055 - 0.055;
+}
+
+half3 LinearToSrgbBranching(half3 lin) 
+{
+	return half3(
+		LinearToSrgbBranchingChannel(lin.r),
+		LinearToSrgbBranchingChannel(lin.g),
+		LinearToSrgbBranchingChannel(lin.b));
+}
+
+half3 sRGBToLinear( half3 Color ) 
+{
+	Color = max(6.10352e-5, Color); // minimum positive non-denormal (fixes black problem on DX11 AMD and NV)
+	return Color > 0.04045 ? pow( Color * (1.0 / 1.055) + 0.0521327, 2.4 ) : Color * (1.0 / 12.92);
+}
+
+/**
+ * @param GammaCurveRatio The curve ratio compared to a 2.2 standard gamma, e.g. 2.2 / DisplayGamma.  So normally the value is 1.
+ */
+half3 ApplyGammaCorrection(float3 LinearColor, float GammaCurveRatio)
+{
+	// Apply "gamma" curve adjustment.
+	float3 CorrectedColor = pow(LinearColor, GammaCurveRatio);
+
+	#if MAC
+		// Note, MacOSX native output is raw gamma 2.2 not sRGB!
+		CorrectedColor = pow(CorrectedColor, 1.0/2.2);
+	#else
+		#if USE_709
+			// Didn't profile yet if the branching version would be faster (different linear segment).
+			CorrectedColor = LinearTo709Branchless(CorrectedColor);
+		#else
+			CorrectedColor = LinearToSrgbBranching(CorrectedColor);
+		#endif
+	#endif
+
+	return CorrectedColor;
+}
\ No newline at end of file
diff --git a/Engine/Shaders/StandaloneRenderer/D3D/SlateDefaultVertexShader.hlsl b/Engine/Shaders/StandaloneRenderer/D3D/SlateDefaultVertexShader.hlsl
--- a/Engine/Shaders/StandaloneRenderer/D3D/SlateDefaultVertexShader.hlsl
+++ b/Engine/Shaders/StandaloneRenderer/D3D/SlateDefaultVertexShader.hlsl
@@ -1,6 +1,6 @@
 // Copyright 1998-2015 Epic Games, Inc. All Rights Reserved.
 
-
+#include "GammaCorrectionCommon.hlsl"
 
 cbuffer PerElementVSConstants
 {
@@ -33,7 +33,8 @@ VertexOut Main(
 	Out.ClipOriginAndPos = float4(InClipOrigin, InPosition.xy);
 	Out.ClipExtents = InClipExtents;
 	
-	Out.Color = InColor;
+	Out.Color.rgb = sRGBToLinear(InColor.rgb);
+	Out.Color.a = InColor.a;
 
 	return Out;
 }
diff --git a/Engine/Shaders/StandaloneRenderer/D3D/SlateElementPixelShader.hlsl b/Engine/Shaders/StandaloneRenderer/D3D/SlateElementPixelShader.hlsl
--- a/Engine/Shaders/StandaloneRenderer/D3D/SlateElementPixelShader.hlsl
+++ b/Engine/Shaders/StandaloneRenderer/D3D/SlateElementPixelShader.hlsl
@@ -1,5 +1,6 @@
 // Copyright 1998-2015 Epic Games, Inc. All Rights Reserved.
 
+#include "GammaCorrectionCommon.hlsl"
 
 // Shader types
 #define ESlateShader::Default		0
@@ -14,6 +15,12 @@
 Texture2D ElementTexture;
 SamplerState ElementTextureSampler;
 
+cbuffer PerFramePSConstants
+{
+	/** Display gamma x:gamma curve adjustment, y:inverse gamma (1/GEngine->DisplayGamma) */
+	float2 GammaValues;
+};
+
 cbuffer PerElementPSConstants
 {
 	float4 ShaderParams;
@@ -34,10 +41,22 @@ struct VertexOut
 
 float3 Hue( float H )
 {
-    float R = abs(H * 6 - 3) - 1;
-    float G = 2 - abs(H * 6 - 2);
-    float B = 2 - abs(H * 6 - 4);
-    return saturate( float3(R,G,B) );
+	float R = abs(H * 6 - 3) - 1;
+	float G = 2 - abs(H * 6 - 2);
+	float B = 2 - abs(H * 6 - 4);
+	return saturate( float3(R,G,B) );
+}
+
+float3 GammaCorrect(float3 InColor)
+{
+	float3 CorrectedColor = InColor;
+
+	if ( GammaValues.y != 1.0f )
+	{
+		CorrectedColor = ApplyGammaCorrection(CorrectedColor, GammaValues.x);
+	}
+
+	return CorrectedColor;
 }
 
 float4 GetFontElementColor( VertexOut InVertex )
@@ -54,10 +73,10 @@ float4 GetColor( VertexOut InVertex, float2 UV )
 	float4 FinalColor;
 	
 	float4 BaseColor = ElementTexture.Sample(ElementTextureSampler, UV );
-    if( ( DrawEffects & DE_IgnoreTextureAlpha ) != 0 )
-    {
-        BaseColor.a = 1.0f;
-    }
+	if( ( DrawEffects & DE_IgnoreTextureAlpha ) != 0 )
+	{
+		BaseColor.a = 1.0f;
+	}
 
 	FinalColor = BaseColor*InVertex.Color;
 	return FinalColor;
@@ -213,7 +232,7 @@ float4 Main( VertexOut InVertex ) : SV_Target
 	}
 
 	// gamma correct
-	OutColor.rgb = pow(OutColor.rgb,1.0f/2.2f);
+	OutColor.rgb = GammaCorrect(OutColor.rgb);
 
 	if( DrawEffects & DE_Disabled )
 	{
diff --git a/Engine/Shaders/StandaloneRenderer/OpenGL/SlateElementPixelShader.glsl b/Engine/Shaders/StandaloneRenderer/OpenGL/SlateElementPixelShader.glsl
--- a/Engine/Shaders/StandaloneRenderer/OpenGL/SlateElementPixelShader.glsl
+++ b/Engine/Shaders/StandaloneRenderer/OpenGL/SlateElementPixelShader.glsl
@@ -8,12 +8,19 @@ precision highp float;
 #extension GL_EXT_gpu_shader4 : enable
 #endif
 
+#ifndef USE_709
+#define USE_709 0
+#endif // USE_709
+
 // Shader types
 #define ST_Default	0
 #define ST_Border	1
 #define ST_Font		2
 #define ST_Line		3
 
+/** Display gamma x:gamma curve adjustment, y:inverse gamma (1/GEngine->DisplayGamma) */
+uniform vec2 GammaValues = vec2(1, 1/2.2);
+
 // Draw effects
 uniform bool EffectsDisabled;
 uniform bool IgnoreTextureAlpha;
@@ -27,6 +34,103 @@ varying vec4 ClipOriginAndPos;
 varying vec4 ClipExtents;
 varying vec4 Color;
 
+vec3 max(float test, vec3 values)
+{
+	return vec3(max(test, values.x), max(test, values.y), max(test, values.z));
+}
+
+vec3 pow(vec3 values, float power)
+{
+	return vec3(pow(values.x, power), pow(values.y, power), pow(values.z, power));
+}
+
+vec3 LinearTo709Branchless(vec3 lin)
+{
+	lin = max(6.10352e-5, lin); // minimum positive non-denormal (fixes black problem on DX11 AMD and NV)
+	return min(lin * 4.5, pow(max(lin, 0.018), 0.45) * 1.099 - 0.099);
+}
+
+vec3 LinearToSrgbBranchless(vec3 lin)
+{
+	lin = max(6.10352e-5, lin); // minimum positive non-denormal (fixes black problem on DX11 AMD and NV)
+	return min(lin * 12.92, pow(max(lin, 0.00313067), 1.0/2.4) * 1.055 - 0.055);
+	// Possible that mobile GPUs might have native pow() function?
+	//return min(lin * 12.92, exp2(log2(max(lin, 0.00313067)) * (1.0/2.4) + log2(1.055)) - 0.055);
+}
+
+float LinearToSrgbBranchingChannel(float lin)
+{
+	if(lin < 0.00313067) return lin * 12.92;
+	return pow(lin, (1.0/2.4)) * 1.055 - 0.055;
+}
+
+vec3 LinearToSrgbBranching(vec3 lin)
+{
+	return vec3(
+				LinearToSrgbBranchingChannel(lin.r),
+				LinearToSrgbBranchingChannel(lin.g),
+				LinearToSrgbBranchingChannel(lin.b));
+}
+
+float sRGBToLinearChannel( float ColorChannel )
+{
+	return ColorChannel > 0.04045 ? pow( ColorChannel * (1.0 / 1.055) + 0.0521327, 2.4 ) : ColorChannel * (1.0 / 12.92);
+}
+
+vec3 sRGBToLinear( vec3 Color )
+{
+	return vec3(sRGBToLinearChannel(Color.r),
+				sRGBToLinearChannel(Color.g),
+				sRGBToLinearChannel(Color.b));
+}
+
+/**
+ * @param GammaCurveRatio The curve ratio compared to a 2.2 standard gamma, e.g. 2.2 / DisplayGamma.  So normally the value is 1.
+ */
+vec3 ApplyGammaCorrection(vec3 LinearColor, float GammaCurveRatio)
+{
+	// Apply "gamma" curve adjustment.
+	vec3 CorrectedColor = pow(LinearColor, GammaCurveRatio);
+
+#if PLATFORM_MAC
+	// Note, MacOSX native output is raw gamma 2.2 not sRGB!
+	//CorrectedColor = pow(CorrectedColor, 1.0/2.2);
+	CorrectedColor = LinearToSrgbBranching(CorrectedColor);
+#else
+#if USE_709
+	// Didn't profile yet if the branching version would be faster (different linear segment).
+	CorrectedColor = LinearTo709Branchless(CorrectedColor);
+#else
+	CorrectedColor = LinearToSrgbBranching(CorrectedColor);
+#endif
+	
+#endif
+	
+	return CorrectedColor;
+}
+
+vec3 GammaCorrect(vec3 InColor)
+{
+	vec3 CorrectedColor = InColor;
+	
+	// gamma correct
+	//#if PLATFORM_USES_ES2
+	//	OutColor.rgb = sqrt( OutColor.rgb );
+	//#else
+	//	OutColor.rgb = pow(OutColor.rgb, vec3(1.0/2.2));
+	//#endif
+	//!(ES2_PROFILE || ES3_1_PROFILE)
+	
+#if !PLATFORM_USES_ES2
+	if( GammaValues.y != 1.0f )
+	{
+		CorrectedColor = ApplyGammaCorrection(CorrectedColor, GammaValues.x);
+	}
+#endif
+	
+	return CorrectedColor;
+}
+
 vec4 GetFontElementColor()
 {
 	vec4 OutColor = Color;
@@ -43,11 +147,11 @@ vec4 GetDefaultElementColor()
 {
 	vec4 OutColor = Color;
 
-    vec4 TextureColor = texture2D(ElementTexture, TexCoords.xy*TexCoords.zw);
+	vec4 TextureColor = texture2D(ElementTexture, TexCoords.xy*TexCoords.zw);
 	if( IgnoreTextureAlpha )
-    {
-        TextureColor.a = 1.0;
-    }
+	{
+		TextureColor.a = 1.0;
+	}
 	OutColor *= TextureColor;
 
 	return OutColor;
@@ -216,12 +320,8 @@ void main()
 	}
 	
 	// gamma correct
-	#if PLATFORM_USES_ES2
-		OutColor.rgb = sqrt( OutColor.rgb );
-	#else
-		OutColor.rgb = pow(OutColor.rgb, vec3(1.0/2.2));
-	#endif
-
+	OutColor.rgb = GammaCorrect(OutColor.rgb);
+	
 	if( EffectsDisabled )
 	{
 		//desaturate
diff --git a/Engine/Shaders/StandaloneRenderer/OpenGL/SlateVertexShader.glsl b/Engine/Shaders/StandaloneRenderer/OpenGL/SlateVertexShader.glsl
--- a/Engine/Shaders/StandaloneRenderer/OpenGL/SlateVertexShader.glsl
+++ b/Engine/Shaders/StandaloneRenderer/OpenGL/SlateVertexShader.glsl
@@ -17,12 +17,32 @@ varying vec4 ClipOriginAndPos;
 varying vec4 ClipExtents;
 varying vec4 Color;
 
+vec3 pow(vec3 values, float power)
+{
+	return vec3(pow(values.x, power), pow(values.y, power), pow(values.z, power));
+}
+
+float sRGBToLinearChannel( float ColorChannel )
+{
+	return ColorChannel > 0.04045 ? pow( ColorChannel * (1.0 / 1.055) + 0.0521327, 2.4 ) : ColorChannel * (1.0 / 12.92);
+}
+
+vec3 sRGBToLinear( vec3 Color )
+{
+	return vec3(
+				sRGBToLinearChannel(Color.r),
+				sRGBToLinearChannel(Color.g),
+				sRGBToLinearChannel(Color.b));
+}
+
 void main()
 {
 	TexCoords = InTexCoords;
 	ClipOriginAndPos = vec4(InClipOrigin, InPosition);
 	ClipExtents = InClipExtents;
-	Color = InColor;
+
+	Color.rgb = sRGBToLinear(InColor.rgb);
+	Color.a = InColor.a;
 
 	gl_Position = ViewProjectionMatrix * vec4( InPosition, 0, 1 );
 }
diff --git a/Engine/Source/Developer/StandaloneRenderer/Private/OpenGL/SlateOpenGLShaders.cpp b/Engine/Source/Developer/StandaloneRenderer/Private/OpenGL/SlateOpenGLShaders.cpp
--- a/Engine/Source/Developer/StandaloneRenderer/Private/OpenGL/SlateOpenGLShaders.cpp
+++ b/Engine/Source/Developer/StandaloneRenderer/Private/OpenGL/SlateOpenGLShaders.cpp
@@ -89,25 +89,51 @@ void FSlateOpenGLShader::CompileShader( const FString& Filename, GLenum ShaderTy
 	bool bFileFound = FFileHelper::LoadFileToString( Source, *Filename );
 	check(bFileFound);
 	
-	// Allocate a buffer big enough to store the string in ascii format
-	ANSICHAR* Chars[2] = {0};
+	FString Header;
+	
 	// pass the #define along to the shader
 #if PLATFORM_USES_ES2
-	Chars[0] = (ANSICHAR*)"#define PLATFORM_USES_ES2 1\n\n#define PLATFORM_LINUX 0\n";
+	Header.Append("#define PLATFORM_USES_ES2 1\n");
 #elif PLATFORM_LINUX
 	#if LINUX_USE_OPENGL_3_2
-	Chars[0] = (ANSICHAR*)"#version 150\n\n#define PLATFORM_USES_ES2 0\n\n#define PLATFORM_LINUX 1\n";
+	Header.Append("#version 150\n#define PLATFORM_USES_ES2 0\n");
 	#else
-	Chars[0] = (ANSICHAR*)"#version 120\n\n#define PLATFORM_USES_ES2 0\n\n#define PLATFORM_LINUX 1\n";
+	Header.Append("#version 120\n#define PLATFORM_USES_ES2 0\n");
 	#endif // LINUX_USE_OPENGL_3_2
 #else
-	Chars[0] = (ANSICHAR*)"#version 120\n\n#define PLATFORM_USES_ES2 0\n\n#define PLATFORM_LINUX 0\n";
+	Header.Append("#version 120\n#define PLATFORM_USES_ES2 0\n");
+#endif
+	
+#if PLATFORM_LINUX
+	Header.Append("#define PLATFORM_LINUX 1\n");
+#else
+	Header.Append("#define PLATFORM_LINUX 0\n");
+#endif
+	
+#if PLATFORM_MAC
+	Header.Append("#define PLATFORM_MAC 1\n");
+#else
+	Header.Append("#define PLATFORM_MAC 0\n");
 #endif
+	
+#if USE_709
+	Header.Append("#define USE_709 1\n");
+#else
+	Header.Append("#define USE_709 0\n");
+#endif
+	
+	// Allocate a buffer big enough to store the string in ascii format
+	ANSICHAR* Chars[2] = {0};
+	
+	Chars[0] = new ANSICHAR[Header.Len()+1];
+	FCStringAnsi::Strcpy(Chars[0], Header.Len() + 1, TCHAR_TO_ANSI(*Header));
+	
 	Chars[1] = new ANSICHAR[Source.Len()+1];
 	FCStringAnsi::Strcpy(Chars[1], Source.Len() + 1, TCHAR_TO_ANSI(*Source));
 
 	// give opengl the source code for the shader
 	glShaderSource( ShaderID, 2, (const ANSICHAR**)Chars, NULL );
+	delete[] Chars[0];
 	delete[] Chars[1];
 
 	// Compile the shader and check for success
diff --git a/Engine/Source/Developer/StandaloneRenderer/Private/Windows/D3D/SlateD3DRenderingPolicy.h b/Engine/Source/Developer/StandaloneRenderer/Private/Windows/D3D/SlateD3DRenderingPolicy.h
--- a/Engine/Source/Developer/StandaloneRenderer/Private/Windows/D3D/SlateD3DRenderingPolicy.h
+++ b/Engine/Source/Developer/StandaloneRenderer/Private/Windows/D3D/SlateD3DRenderingPolicy.h
@@ -14,7 +14,7 @@ public:
 
 	virtual TSharedRef<FSlateShaderResourceManager> GetResourceManager() override;
 	virtual TSharedRef<FSlateFontCache> GetFontCache() override { return FontCache.ToSharedRef(); }
-	virtual bool IsVertexColorInLinearSpace() const override { return true; }
+	virtual bool IsVertexColorInLinearSpace() const override { return false; }
 
 private:
 	void InitResources();
diff --git a/Engine/Source/Developer/StandaloneRenderer/Private/Windows/D3D/SlateD3DShaders.cpp b/Engine/Source/Developer/StandaloneRenderer/Private/Windows/D3D/SlateD3DShaders.cpp
--- a/Engine/Source/Developer/StandaloneRenderer/Private/Windows/D3D/SlateD3DShaders.cpp
+++ b/Engine/Source/Developer/StandaloneRenderer/Private/Windows/D3D/SlateD3DShaders.cpp
@@ -276,13 +276,19 @@ FSlateDefaultPS::FSlateDefaultPS()
 {
 	Texture = &FSlateShaderParameterMap::Get().RegisterParameter<ID3D11ShaderResourceView>( "ElementTexture" );
 	TextureSampler = &FSlateShaderParameterMap::Get().RegisterParameter<ID3D11SamplerState>( "ElementTextureSampler" );
-	PerElementCBufferParam = &FSlateShaderParameterMap::Get().RegisterParameter<ID3D11Buffer>( "PerElementPSConstants" );
+	PerFrameCBufferParam = &FSlateShaderParameterMap::Get().RegisterParameter<ID3D11Buffer>("PerFramePSConstants");
+	PerElementCBufferParam = &FSlateShaderParameterMap::Get().RegisterParameter<ID3D11Buffer>("PerElementPSConstants");
 
+	PerFrameConstants.Create();
 	PerElementConstants.Create();
-	
+
+	PerFrameConstants.GetBufferData().GammaValues = FVector2D(1, 1 / 2.2f);
+
+	PerFrameCBufferParam->SetParameter(PerFrameConstants.GetResource());
+
 	// Set the constant parameter to use our constant buffer
 	// @todo: If we go back to multiple pixel shaders this likely has be called more frequently
-	PerElementCBufferParam->SetParameter( PerElementConstants.GetResource() );
+	PerElementCBufferParam->SetParameter(PerElementConstants.GetResource());
 
 	Create( FString::Printf( TEXT("%sShaders/StandaloneRenderer/D3D/SlateElementPixelShader.hlsl"), *FPaths::EngineDir() ), TEXT("Main"), TEXT("ps_4_0") );
 }
@@ -302,10 +308,15 @@ void FSlateDefaultPS::SetShaderParams( const FVector4& InShaderParams )
 	PerElementConstants.GetBufferData().ShaderParams = InShaderParams;
 }
 
+void FSlateDefaultPS::SetGammaValues(const FVector2D& InGammaValues)
+{
+	PerFrameConstants.GetBufferData().GammaValues = InGammaValues;
+}
+
 void FSlateDefaultPS::UpdateParameters()
 {
+	PerFrameConstants.UpdateBuffer();
 	PerElementConstants.UpdateBuffer();
 
 	TextureSampler->SetParameter( SamplerState );
-
 }
diff --git a/Engine/Source/Developer/StandaloneRenderer/Private/Windows/D3D/SlateD3DShaders.h b/Engine/Source/Developer/StandaloneRenderer/Private/Windows/D3D/SlateD3DShaders.h
--- a/Engine/Source/Developer/StandaloneRenderer/Private/Windows/D3D/SlateD3DShaders.h
+++ b/Engine/Source/Developer/StandaloneRenderer/Private/Windows/D3D/SlateD3DShaders.h
@@ -172,21 +172,27 @@ public:
 	void SetShaderType( uint32 InShaderType );
 	void SetDrawEffects( uint32 InDrawEffects );
 	void SetShaderParams( const FVector4& InShaderParams );
+	void SetGammaValues(const FVector2D& InGammaValues);
 protected:
 	void UpdateParameters();
 private:
-	struct FPerElementConstants
+	MS_ALIGN(16) struct FPerElementConstants
 	{
 		FVector4 ShaderParams;
 		uint32 DrawEffects;
 		uint32 ShaderType;
-		// Constant buffers must be 16 byte aligned
-		uint32 Padding[2];
+	};
+
+	MS_ALIGN(16) struct FPerFrameConstants
+	{
+		FVector2D GammaValues;
 	};
 
 	FSlateD3DConstantBuffer<FPerElementConstants> PerElementConstants;
+	FSlateD3DConstantBuffer<FPerFrameConstants> PerFrameConstants;
 	TSlateD3DTypedShaderParameter<ID3D11ShaderResourceView>* Texture;
 	TSlateD3DTypedShaderParameter<ID3D11SamplerState>* TextureSampler;
+	TSlateD3DTypedShaderParameter<ID3D11Buffer>* PerFrameCBufferParam;
 	TSlateD3DTypedShaderParameter<ID3D11Buffer>* PerElementCBufferParam;
 	TRefCountPtr<ID3D11SamplerState> SamplerState;
 };
diff --git a/Engine/Source/Runtime/AppFramework/Private/Widgets/Testing/STestSuite.cpp b/Engine/Source/Runtime/AppFramework/Private/Widgets/Testing/STestSuite.cpp
--- a/Engine/Source/Runtime/AppFramework/Private/Widgets/Testing/STestSuite.cpp
+++ b/Engine/Source/Runtime/AppFramework/Private/Widgets/Testing/STestSuite.cpp
@@ -3950,6 +3950,44 @@ private:
 	TSharedPtr<SInvalidationPanel> CachePanel1;
 };
 
+class SGammaTest : public SCompoundWidget
+{
+	SLATE_BEGIN_ARGS(SGammaTest)
+	{}
+	SLATE_END_ARGS()
+	
+	void Construct(const FArguments& InArgs)
+	{
+		FColor Orange(200, 80, 15);
+		
+		ChildSlot
+		.Padding(10)
+		[
+			SNew(SVerticalBox)
+			
+			+ SVerticalBox::Slot()
+			.AutoHeight()
+			[
+				SNew(SHorizontalBox)
+			 
+				+ SHorizontalBox::Slot()
+				.AutoWidth()
+				[
+					SNew(SColorBlock)
+					.Color(Orange)
+				]
+			 
+				+ SHorizontalBox::Slot()
+				.AutoWidth()
+				[
+					SNew(STextBlock)
+					.Text(FText::FromString(Orange.ToString()))
+				]
+			]
+		];
+	}
+};
+
 class SColorPickerTest : public SCompoundWidget
 {
 public:
@@ -5331,6 +5369,15 @@ TSharedRef<SDockTab> SpawnTab(const FSpawnTabArgs& Args, FName TabIdentifier)
 				]
 			];
 	}
+	else if ( TabIdentifier == FName(TEXT("GammaTest")) )
+	{
+		return SNew(SDockTab)
+		[
+			SNew(SGammaTest)
+			.RenderTransform_Static(&::GetTestRenderTransform)
+			.RenderTransformPivot_Static(&::GetTestRenderTransformPivot)
+		];
+	}
 	else if (TabIdentifier == FName(TEXT("NotificationListTestTab")))
 	{
 		return SNew(SDockTab)
@@ -5553,6 +5600,7 @@ TSharedRef<SDockTab> SpawnTestSuite2( const FSpawnTabArgs& Args )
 			->AddTab("GridPanelTest", ETabState::OpenedTab)
 			->AddTab("DPIScalingTest", ETabState::OpenedTab)
 			->AddTab("InvalidationTest", ETabState::OpenedTab)
+			->AddTab("GammaTest", ETabState::OpenedTab)
 		)
 	);
 
@@ -5591,6 +5639,10 @@ TSharedRef<SDockTab> SpawnTestSuite2( const FSpawnTabArgs& Args )
 		TestSuite2TabManager->RegisterTabSpawner("InvalidationTest", FOnSpawnTab::CreateStatic(&SpawnTab, FName("InvalidationTest")))
 			.SetDisplayName(NSLOCTEXT("TestSuite1", "InvalidationTest", "Invalidtion"))
 			.SetGroup(TestSuiteMenu::SuiteTabs);
+		
+		TestSuite2TabManager->RegisterTabSpawner("GammaTest", FOnSpawnTab::CreateStatic(&SpawnTab, FName("GammaTest")))
+			.SetDisplayName(NSLOCTEXT("TestSuite1", "GammaTest", "Gamma"))
+			.SetGroup(TestSuiteMenu::SuiteTabs);
 	}
 
 	FMenuBarBuilder MenuBarBuilder = FMenuBarBuilder( TSharedPtr<FUICommandList>() );
diff --git a/Engine/Source/Runtime/SlateCore/Private/Rendering/ElementBatcher.cpp b/Engine/Source/Runtime/SlateCore/Private/Rendering/ElementBatcher.cpp
--- a/Engine/Source/Runtime/SlateCore/Private/Rendering/ElementBatcher.cpp
+++ b/Engine/Source/Runtime/SlateCore/Private/Rendering/ElementBatcher.cpp
@@ -224,6 +224,13 @@ void FSlateElementBatcher::AddElements(FSlateDrawLayer& InDrawLayer)
 	}
 }
 
+FColor FSlateElementBatcher::PackVertexColor(const FLinearColor& InLinearColor)
+{
+	//NOTE: Using pow(x,2) instead of a full sRGB conversion has been tried, but it ended up
+	// causing too much loss of data in the lower levels of black.
+	return InLinearColor.ToFColor(bSRGBVertexColor);
+}
+
 void FSlateElementBatcher::AddQuadElement( const FSlateDrawElement& DrawElement, FColor Color )
 {
 	const FSlateRenderTransform& RenderTransform = DrawElement.GetRenderTransform();
@@ -289,7 +296,7 @@ FSlateRenderTransform GetBoxRenderTransform(const FSlateDrawElement& DrawElement
 void FSlateElementBatcher::AddBoxElement( const FSlateDrawElement& DrawElement )
 {
 	const FSlateDataPayload& InPayload = DrawElement.GetDataPayload();
-	FColor Tint = InPayload.Tint.ToFColor(bSRGBVertexColor);
+	FColor Tint = PackVertexColor(InPayload.Tint);
 
 	if ( Tint.A == 0 )
 	{
@@ -597,7 +604,7 @@ void FSlateElementBatcher::AddBoxElement( const FSlateDrawElement& DrawElement )
 void FSlateElementBatcher::AddTextElement(const FSlateDrawElement& DrawElement)
 {
 	const FSlateDataPayload& InPayload = DrawElement.GetDataPayload();
-	FColor Tint = InPayload.Tint.ToFColor(bSRGBVertexColor);
+	FColor Tint = PackVertexColor(InPayload.Tint);
 
 	if ( Tint.A == 0 )
 	{
@@ -975,7 +982,7 @@ void FSlateElementBatcher::AddSplineElement( const FSlateDrawElement& DrawElemen
 	FVector2D StartPos = StartPt;
 	FVector2D EndPos = FVector2D( FMath::CubicInterp( StartPt, StartDir, EndPt, EndDir, Alpha ) );
 
-	FColor FinalColor = InPayload.Tint.ToFColor(bSRGBVertexColor);
+	FColor FinalColor = PackVertexColor(InPayload.Tint);
 
 	BatchVertices.Add( FSlateVertex( RenderTransform, StartPos + Up, TransformPoint(RenderTransform, StartPos), TransformPoint(RenderTransform, EndPos), FinalColor, RenderClipRect ) );
 	BatchVertices.Add( FSlateVertex( RenderTransform, StartPos - Up, TransformPoint(RenderTransform, StartPos), TransformPoint(RenderTransform, EndPos), FinalColor, RenderClipRect ) );
@@ -1078,7 +1085,7 @@ void FSlateElementBatcher::AddLineElement( const FSlateDrawElement& DrawElement
 		return;
 	}
 
-	FColor FinalColor = InPayload.Tint.ToFColor(bSRGBVertexColor);
+	FColor FinalColor = PackVertexColor(InPayload.Tint);
 
 	if( InPayload.bAntialias )
 	{
@@ -1245,7 +1252,7 @@ void FSlateElementBatcher::AddViewportElement( const FSlateDrawElement& DrawElem
 	FSlateLayoutTransform InverseLayoutTransform(Inverse(FSlateLayoutTransform(DrawElement.GetScale(), DrawElement.GetPosition())));
 	FSlateRotatedClipRectType RenderClipRect = ToSnappedRotatedRect(InClippingRect, InverseLayoutTransform, RenderTransform);
 
-	const FColor FinalColor = InPayload.Tint.ToFColor(bSRGBVertexColor);
+	const FColor FinalColor = PackVertexColor(InPayload.Tint);
 
 	ESlateBatchDrawFlag::Type DrawFlags = ESlateBatchDrawFlag::None;
 	
@@ -1399,7 +1406,7 @@ void FSlateElementBatcher::AddBorderElement( const FSlateDrawElement& DrawElemen
 	FShaderParams ShaderParams = FShaderParams::MakePixelShaderParams( FVector4(LeftMarginU,RightMarginU,TopMarginV,BottomMarginV) );
 
 	// The tint color applies to all brushes and is passed per vertex
-	FColor Tint = InPayload.Tint.ToFColor(bSRGBVertexColor);
+	const FColor Tint = PackVertexColor(InPayload.Tint);
 
 	// Pass the tiling information as a flag so we can pick the correct texture addressing mode
 	ESlateBatchDrawFlag::Type DrawFlags = (ESlateBatchDrawFlag::TileU|ESlateBatchDrawFlag::TileV);
diff --git a/Engine/Source/Runtime/SlateCore/Public/Rendering/ElementBatcher.h b/Engine/Source/Runtime/SlateCore/Public/Rendering/ElementBatcher.h
--- a/Engine/Source/Runtime/SlateCore/Public/Rendering/ElementBatcher.h
+++ b/Engine/Source/Runtime/SlateCore/Public/Rendering/ElementBatcher.h
@@ -34,6 +34,8 @@ public:
 
 private:
 	void AddElements(FSlateDrawLayer& InDrawLayer);
+	
+	FColor PackVertexColor(const FLinearColor& InLinearColor);
 
 	/** 
 	 * Creates vertices necessary to draw a Quad element 
