SlateViewer: fix GLSL shaders on Linux.

diff --git a/Engine/Shaders/StandaloneRenderer/OpenGL/SlateElementPixelShader.glsl b/Engine/Shaders/StandaloneRenderer/OpenGL/SlateElementPixelShader.glsl
--- a/Engine/Shaders/StandaloneRenderer/OpenGL/SlateElementPixelShader.glsl
+++ b/Engine/Shaders/StandaloneRenderer/OpenGL/SlateElementPixelShader.glsl
@@ -34,26 +34,26 @@ varying vec4 ClipOriginAndPos;
 varying vec4 ClipExtents;
 varying vec4 Color;
 
-vec3 max(float test, vec3 values)
+vec3 maxWithScalar(float test, vec3 values)
 {
 	return vec3(max(test, values.x), max(test, values.y), max(test, values.z));
 }
 
-vec3 pow(vec3 values, float power)
+vec3 powScalar(vec3 values, float power)
 {
 	return vec3(pow(values.x, power), pow(values.y, power), pow(values.z, power));
 }
 
 vec3 LinearTo709Branchless(vec3 lin)
 {
-	lin = max(6.10352e-5, lin); // minimum positive non-denormal (fixes black problem on DX11 AMD and NV)
-	return min(lin * 4.5, pow(max(lin, 0.018), 0.45) * 1.099 - 0.099);
+	lin = maxWithScalar(6.10352e-5, lin); // minimum positive non-denormal (fixes black problem on DX11 AMD and NV)
+	return min(lin * 4.5, powScalar(maxWithScalar(0.018, lin), 0.45) * 1.099 - 0.099);
 }
 
 vec3 LinearToSrgbBranchless(vec3 lin)
 {
-	lin = max(6.10352e-5, lin); // minimum positive non-denormal (fixes black problem on DX11 AMD and NV)
-	return min(lin * 12.92, pow(max(lin, 0.00313067), 1.0/2.4) * 1.055 - 0.055);
+	lin = maxWithScalar(6.10352e-5, lin); // minimum positive non-denormal (fixes black problem on DX11 AMD and NV)
+	return min(lin * 12.92, powScalar(maxWithScalar(0.00313067, lin), 1.0/2.4) * 1.055 - 0.055);
 	// Possible that mobile GPUs might have native pow() function?
 	//return min(lin * 12.92, exp2(log2(max(lin, 0.00313067)) * (1.0/2.4) + log2(1.055)) - 0.055);
 }
@@ -90,7 +90,7 @@ vec3 sRGBToLinear( vec3 Color )
 vec3 ApplyGammaCorrection(vec3 LinearColor, float GammaCurveRatio)
 {
 	// Apply "gamma" curve adjustment.
-	vec3 CorrectedColor = pow(LinearColor, GammaCurveRatio);
+	vec3 CorrectedColor = powScalar(LinearColor, GammaCurveRatio);
 
 #if PLATFORM_MAC
 	// Note, MacOSX native output is raw gamma 2.2 not sRGB!
diff --git a/Engine/Shaders/StandaloneRenderer/OpenGL/SlateVertexShader.glsl b/Engine/Shaders/StandaloneRenderer/OpenGL/SlateVertexShader.glsl
--- a/Engine/Shaders/StandaloneRenderer/OpenGL/SlateVertexShader.glsl
+++ b/Engine/Shaders/StandaloneRenderer/OpenGL/SlateVertexShader.glsl
@@ -17,7 +17,7 @@ varying vec4 ClipOriginAndPos;
 varying vec4 ClipExtents;
 varying vec4 Color;
 
-vec3 pow(vec3 values, float power)
+vec3 powScalar(vec3 values, float power)
 {
 	return vec3(pow(values.x, power), pow(values.y, power), pow(values.z, power));
 }
