Fixed macros for GLES in Slate shaders.

#rnx

diff --git a/Engine/Shaders/StandaloneRenderer/OpenGL/SlateElementPixelShader.glsl b/Engine/Shaders/StandaloneRenderer/OpenGL/SlateElementPixelShader.glsl
--- a/Engine/Shaders/StandaloneRenderer/OpenGL/SlateElementPixelShader.glsl
+++ b/Engine/Shaders/StandaloneRenderer/OpenGL/SlateElementPixelShader.glsl
@@ -1,7 +1,7 @@
 // Copyright Epic Games, Inc. All Rights Reserved.
 
 // handle differences between ES2 and full GL shaders
-#if PLATFORM_USES_ES2
+#if PLATFORM_USES_GLES
 precision highp float;
 #else
 // #version 120 at the beginning is added in FSlateOpenGLShader::CompileShader()
@@ -114,14 +114,14 @@ vec3 GammaCorrect(vec3 InColor)
 	vec3 CorrectedColor = InColor;
 	
 	// gamma correct
-	//#if PLATFORM_USES_ES2
+	//#if PLATFORM_USES_GLES
 	//	OutColor.rgb = sqrt( OutColor.rgb );
 	//#else
 	//	OutColor.rgb = pow(OutColor.rgb, vec3(1.0/2.2));
 	//#endif
 	//!(ES2_PROFILE || ES3_1_PROFILE)
 	
-#if !PLATFORM_USES_ES2
+#if !PLATFORM_USES_GLES
 	if( GammaValues.y != 1.0f )
 	{
 		CorrectedColor = ApplyGammaCorrection(CorrectedColor, GammaValues.x);
diff --git a/Engine/Shaders/StandaloneRenderer/OpenGL/SplashFragmentShader.glsl b/Engine/Shaders/StandaloneRenderer/OpenGL/SplashFragmentShader.glsl
--- a/Engine/Shaders/StandaloneRenderer/OpenGL/SplashFragmentShader.glsl
+++ b/Engine/Shaders/StandaloneRenderer/OpenGL/SplashFragmentShader.glsl
@@ -1,5 +1,5 @@
 
-#if PLATFORM_USES_ES2
+#if PLATFORM_USES_GLES
 precision highp float;
 #else
 // #version 120 at the beginning is added in FSlateOpenGLShader::CompileShader()
@@ -17,4 +17,4 @@ void main()
 
 	gl_FragColor = texture2D(SplashTexture, tc);
 
-}
\ No newline at end of file
+}
