Change disabled effect for slate to a simple alpha fade. This works better for grayscale while also preserving colors better than the old effect. The old effect is still available for UI that depended on its look

diff --git a/Engine/Shaders/Private/SlateElementPixelShader.usf b/Engine/Shaders/Private/SlateElementPixelShader.usf
--- a/Engine/Shaders/Private/SlateElementPixelShader.usf
+++ b/Engine/Shaders/Private/SlateElementPixelShader.usf
@@ -10,6 +10,9 @@ half4 GammaAndAlphaValues;
 /** x:DrawDisabledEffect */
 /** Per element params */
 float4 ShaderParams;
+float4 ShaderParams2;
+
+#define USE_LEGACY_DISABLED_EFFECT 0
 
 #if USE_MATERIALS
 half4 DrawFlags;
@@ -228,6 +231,54 @@ half4 GetBorderElementColor( VertexToPixelInterpolants InVertex )
 }
 #endif
 
+#if SHADER_TYPE == ST_RoundedBox
+float GetRoundedBoxDistance(float2 pos, float2 center, float radius, float inset)
+{
+	// distance from center
+    pos = abs(pos - center); 
+
+    // distance from the inner corner
+    pos = pos - (center - float2(radius + inset, radius + inset));
+
+    // use distance to nearest edge when not in quadrant with radius
+    // this handles an edge case when radius is very close to thickness
+    // otherwise we're in the quadrant with the radius, 
+    // just use the analytic signed distance function
+    return lerp( length(pos) - radius, max(pos.x - radius, pos.y - radius), float(pos.x <= 0 || pos.y <=0) );
+}
+
+half4 GetRoundedBoxElementColor( VertexToPixelInterpolants InVertex )
+{
+	const float2 size = ShaderParams.zw;
+	float2 pos = size * InVertex.TextureCoordinates[0].xy;
+	float2 center = size / 2.0;
+
+	float radius = ShaderParams.x;	
+	float thickness = ShaderParams.y; 
+
+	// Compute the distances internal and external to the border outline
+	float dext = GetRoundedBoxDistance(pos, center, radius, 0.0);
+	float din  = GetRoundedBoxDistance(pos, center, max(radius - thickness, 0), thickness);
+
+	// Compute the border intensity and fill intensity with a smooth transition
+	float spread = 0.5;
+    float bi = smoothstep(spread, -spread, dext);
+    float fi = smoothstep(spread, -spread, din);
+
+	half4 fill = InVertex.Color;
+    half4 border = ShaderParams2;
+
+    // alpha blend the external color 
+    half4 OutColor = lerp(border, fill, half(thickness > radius));
+    OutColor.a = 0.0;
+
+    // blend in the border and fill colors
+    OutColor = lerp(OutColor, border, bi);
+    OutColor = lerp(OutColor, fill, fi);
+    return OutColor;
+}
+#endif
+
 #if SHADER_TYPE == ST_LineSegment 
 /** 
  * Generates an anti-aliased line segment pixel
@@ -261,6 +312,8 @@ half4 Main( VertexToPixelInterpolants VIn ) : SV_Target0
 	half4 OutColor;
 #if SHADER_TYPE == ST_Default || SHADER_TYPE == ST_Custom
 	OutColor = GetDefaultElementColor( VIn );
+#elif SHADER_TYPE == ST_RoundedBox
+	OutColor = GetRoundedBoxElementColor( VIn );
 #elif SHADER_TYPE == ST_Border
 	OutColor = GetBorderElementColor( VIn );
 #elif SHADER_TYPE == ST_GrayscaleFont || SHADER_TYPE == ST_ColorFont
@@ -284,6 +337,7 @@ half4 Main( VertexToPixelInterpolants VIn ) : SV_Target0
 
 	if (DrawDisabledEffect != 0)
 	{
+#if USE_LEGACY_DISABLED_EFFECT
 		#if SOURCE_IN_LINEAR_SPACE
 			half3 LumCoeffs = half3( 0.3, 0.59, .11 );
 			half3 Grayish = half3(.1, .1, .1);
@@ -297,6 +351,12 @@ half4 Main( VertexToPixelInterpolants VIn ) : SV_Target0
 		half Lum = dot( LumCoeffs, OutColor.rgb );
 		OutColor.rgb = lerp( OutColor.rgb, float3(Lum,Lum,Lum), .8 );
 		OutColor.rgb = lerp( OutColor.rgb, Grayish, clamp( distance( OutColor.rgb, Grayish ), 0, .8)  );
+#else
+		//float bw = (min(OutColor.r, min(OutColor.g, OutColor.b)) + max(OutColor.r, max(OutColor.g, OutColor.b))) * 0.5f;
+		//OutColor.rgb = lerp(OutColor.rgb, float3(bw, bw, bw), .5);
+
+		OutColor.a *= .45f;
+#endif
 	}
 
     // Contrast
diff --git a/Engine/Shaders/StandaloneRenderer/D3D/SlateElementPixelShader.hlsl b/Engine/Shaders/StandaloneRenderer/D3D/SlateElementPixelShader.hlsl
--- a/Engine/Shaders/StandaloneRenderer/D3D/SlateElementPixelShader.hlsl
+++ b/Engine/Shaders/StandaloneRenderer/D3D/SlateElementPixelShader.hlsl
@@ -8,6 +8,9 @@
 #define ESlateShader::GrayscaleFont	2
 #define ESlateShader::ColorFont		3
 #define ESlateShader::LineSegment	4
+#define ESlateShader::RoundedBox    7
+
+#define USE_LEGACY_DISABLED_EFFECT 0
 
 Texture2D ElementTexture;
 SamplerState ElementTextureSampler;
@@ -20,10 +23,11 @@ cbuffer PerFramePSConstants
 
 cbuffer PerElementPSConstants
 {
-    uint ShaderType;            //  4 bytes
     float4 ShaderParams;        // 16 bytes
-    uint IgnoreTextureAlpha;    //	4 byte
-    uint DisableEffect;         //  4 byte
+	float4 ShaderParams2;		// 16 bytes
+	uint ShaderType;            //  4 bytes
+    uint IgnoreTextureAlpha;    //	4 bytes
+    uint DisableEffect;         //  4 bytes
     uint UNUSED[1];             //  4 bytes
 };
 
@@ -124,6 +128,51 @@ float4 GetBorderElementColor( VertexOut InVertex )
 	return GetColor( InVertex, NewUV );
 }
 
+float GetRoundedBoxDistance(float2 pos, float2 center, float radius, float inset)
+{
+	// distance from center
+    pos = abs(pos - center); 
+
+    // distance from the inner corner
+    pos = pos - (center - float2(radius + inset, radius + inset));
+
+    // use distance to nearest edge when not in quadrant with radius
+    // this handles an edge case when radius is very close to thickness
+    // otherwise we're in the quadrant with the radius, 
+    // just use the analytic signed distance function
+    return lerp( length(pos) - radius, max(pos.x - radius, pos.y - radius), float(pos.x <= 0 || pos.y <=0) );
+}
+
+float4 GetRoundedBoxElementColor( VertexOut InVertex )
+{
+	const float2 size = ShaderParams.zw;
+	float2 pos = size * InVertex.TextureCoordinates.xy;
+	float2 center = size / 2.0;
+
+	float radius = ShaderParams.x;	
+	float thickness = ShaderParams.y; 
+
+	// Compute the distances internal and external to the border outline
+	float dext = GetRoundedBoxDistance(pos, center, radius, 0.0);
+	float din  = GetRoundedBoxDistance(pos, center, max(radius - thickness, 0), thickness);
+
+	// Compute the border intensity and fill intensity with a smooth transition
+	float spread = 0.5;
+    float bi = smoothstep(spread, -spread, dext);
+    float fi = smoothstep(spread, -spread, din);
+
+    // alpha blend the external color 
+	float4 fill = InVertex.Color;
+    float4 border = ShaderParams2;
+    float4 OutColor = lerp(border, fill, float(thickness > radius));
+    OutColor.a = 0.0;
+
+    // blend in the border and fill colors
+    OutColor = lerp(OutColor, border, bi);
+    OutColor = lerp(OutColor, fill, fi);
+    return OutColor;
+}
+
 float4 GetSplineElementColor( VertexOut InVertex )
 {
 	const float LineWidth = ShaderParams.x;
@@ -156,6 +205,10 @@ float4 Main( VertexOut InVertex ) : SV_Target
 	{
 		OutColor = GetDefaultElementColor( InVertex );
 	}
+	else if( ShaderType == ESlateShader::RoundedBox)
+	{
+		OutColor = GetRoundedBoxElementColor( InVertex );
+	}
 	else if( ShaderType == ESlateShader::Border )
 	{
 		OutColor = GetBorderElementColor( InVertex );
@@ -178,6 +231,8 @@ float4 Main( VertexOut InVertex ) : SV_Target
 
     if (DisableEffect)
 	{
+#if USE_LEGACY_DISABLED_EFFECT
+
 		//desaturate
 		float3 LumCoeffs = float3( 0.3, 0.59, .11 );
 		float Lum = dot( LumCoeffs, OutColor.rgb );
@@ -187,6 +242,9 @@ float4 Main( VertexOut InVertex ) : SV_Target
 		
 		// lerp between desaturated color and gray color based on distance from the desaturated color to the gray
 		OutColor.rgb = lerp( OutColor.rgb, Grayish, clamp( distance( OutColor.rgb, Grayish ), 0, .8)  );
+#else
+		OutColor.a *= .45f;
+#endif
 	}
 
 	return OutColor;
diff --git a/Engine/Shaders/StandaloneRenderer/OpenGL/SlateElementPixelShader.glsl b/Engine/Shaders/StandaloneRenderer/OpenGL/SlateElementPixelShader.glsl
--- a/Engine/Shaders/StandaloneRenderer/OpenGL/SlateElementPixelShader.glsl
+++ b/Engine/Shaders/StandaloneRenderer/OpenGL/SlateElementPixelShader.glsl
@@ -18,6 +18,9 @@ precision highp float;
 #define ST_GrayscaleFont	2
 #define ST_ColorFont		3
 #define ST_Line				4
+#define ST_RoundedBox       7
+
+#define USE_LEGACY_DISABLED_EFFECT 0
 
 /** Display gamma x:gamma curve adjustment, y:inverse gamma (1/GEngine->DisplayGamma) */
 uniform vec2 GammaValues = vec2(1, 1/2.2);
@@ -26,7 +29,8 @@ uniform vec2 GammaValues = vec2(1, 1/2.2);
 uniform bool EffectsDisabled;
 uniform bool IgnoreTextureAlpha;
 
-uniform vec4 MarginUVs;
+uniform vec4 ShaderParams;
+uniform vec4 ShaderParams2;
 uniform int ShaderType;
 uniform sampler2D ElementTexture;
 
@@ -162,7 +166,6 @@ vec4 GetDefaultElementColor()
 		TextureColor.a = 1.0;
 	}
 	OutColor *= TextureColor;
-
 	return OutColor;
 }
 
@@ -182,14 +185,14 @@ vec4 GetBorderElementColor()
 	
 		if( InTexCoords.z > 0.0 )
 		{
-			MinUV = vec2(MarginUVs.x,0.0);
-			MaxUV = vec2(MarginUVs.y,1.0);
+			MinUV = vec2(ShaderParams.x,0.0);
+			MaxUV = vec2(ShaderParams.y,1.0);
 			InTexCoords.w = 1.0;
 		}
 		else
 		{
-			MinUV = vec2(0.0,MarginUVs.z);
-			MaxUV = vec2(1.0,MarginUVs.w);
+			MinUV = vec2(0.0,ShaderParams.z);
+			MaxUV = vec2(1.0,ShaderParams.w);
 			InTexCoords.z = 1.0;
 		}
 
@@ -206,14 +209,60 @@ vec4 GetBorderElementColor()
 	}
 		
 	OutColor *= TextureColor;
-
 	return OutColor;
 }
 
+float GetRoundedBoxDistance(vec2 pos, vec2 center, float radius, float inset)
+{
+	// distance from center
+    pos = abs(pos - center); 
+
+    // distance from the inner corner
+    pos = pos - (center - vec2(radius + inset, radius + inset));
+
+    // use distance to nearest edge when not in quadrant with radius
+    // this handles an edge case when radius is very close to thickness
+    // otherwise we're in the quadrant with the radius, 
+    // just use the analytic signed distance function
+    return mix( length(pos) - radius,
+    			max(pos.x - radius, pos.y - radius), 
+    			float(pos.x <= 0 || pos.y <=0) );
+}
+
+vec4 GetRoundedBoxElementColor()
+{
+	vec2 size = ShaderParams.zw;
+	vec2 pos = size * TexCoords.xy;
+	vec2 center = size / 2.0;
+
+	float radius = ShaderParams.x;	
+	float thickness = ShaderParams.y; 
+
+	// Compute the distances internal and external to the border outline
+	float dext = GetRoundedBoxDistance(pos, center, radius, 0.0);
+	float din  = GetRoundedBoxDistance(pos, center, max(radius - thickness, 0), thickness);
+
+	// Compute the border intensity and fill intensity with a smooth transition
+	float spread = 0.5;
+    float bi = smoothstep(spread, -spread, dext);
+    float fi = smoothstep(spread, -spread, din);
+
+    // alpha blend the external color 
+	vec4 fill = Color;
+    vec4 border = ShaderParams2.bgra;
+    vec4 OutColor = mix(border, fill, float(thickness > radius));
+    OutColor.a = 0.0;
+
+    // blend in the border and fill colors
+    OutColor = mix(OutColor, border, bi);
+    OutColor = mix(OutColor, fill, fi);
+    return OutColor;
+}
+
 vec4 GetSplineElementColor()
 {
-	float LineWidth = MarginUVs.x;
-	float FilterWidthScale = MarginUVs.y;
+	float LineWidth = ShaderParams.x;
+	float FilterWidthScale = ShaderParams.y;
 
 	float Gradient = TexCoords.x;
 	vec2 GradientDerivative = vec2(abs(dFdx(Gradient)), abs(dFdy(Gradient)));
@@ -242,6 +291,10 @@ void main()
 	{
 		OutColor = GetDefaultElementColor();
 	}
+	else if( ShaderType == ST_RoundedBox )
+	{
+		OutColor = GetRoundedBoxElementColor();
+	}
 	else if( ShaderType == ST_Border )
 	{
 		OutColor = GetBorderElementColor();
@@ -264,6 +317,7 @@ void main()
 	
 	if( EffectsDisabled )
 	{
+	#if USE_LEGACY_DISABLED_EFFECT
 		//desaturate
 		vec3 LumCoeffs = vec3( 0.3, 0.59, .11 );
 		float Lum = dot( LumCoeffs, OutColor.rgb );
@@ -272,6 +326,9 @@ void main()
 		vec3 Grayish = vec3(0.4, 0.4, 0.4);
 	
 		OutColor.rgb = mix( OutColor.rgb, Grayish, clamp( distance( OutColor.rgb, Grayish ), 0.0, 0.8)  );
+	#else
+		OutColor.a *= .45f;
+	#endif
 	}
 
 	gl_FragColor = OutColor.bgra;
