Revert MSDF changes to standalone renderer shaders
#rnx <Release Notes exclude>

diff --git a/Engine/Shaders/StandaloneRenderer/D3D/SlateElementPixelShader.hlsl b/Engine/Shaders/StandaloneRenderer/D3D/SlateElementPixelShader.hlsl
--- a/Engine/Shaders/StandaloneRenderer/D3D/SlateElementPixelShader.hlsl
+++ b/Engine/Shaders/StandaloneRenderer/D3D/SlateElementPixelShader.hlsl
@@ -9,8 +9,6 @@
 #define ESlateShader::ColorFont		3
 #define ESlateShader::LineSegment	4
 #define ESlateShader::RoundedBox    7
-#define ESlateShader::SdfFont       8
-#define ESlateShader::MsdfFont      9
 
 #define USE_LEGACY_DISABLED_EFFECT 0
 
@@ -79,35 +77,6 @@ float4 GetColorFontElementColor(VertexOut InVertex)
 	return OutColor;
 }
 
-float GetOpacityFromSignedDistance(float SampledDistance, float2 SdfCoord)
-{
-	const float2 ScreenSpaceTextureSize = float2(1.0, 1.0) / fwidth(SdfCoord);
-	const float ScreenSpaceSpread = max(dot(ShaderParams.xy, ScreenSpaceTextureSize), 1.0);
-	const float ScreenSpaceDistance = (SampledDistance - ShaderParams.z) * ScreenSpaceSpread;
-	return saturate(ScreenSpaceDistance + 0.5);
-}
-
-float4 GetSdfFontElementColor(VertexOut InVertex)
-{
-	float4 OutColor = InVertex.Color;
-
-	const float SampledDistance = float(ElementTexture.Sample(ElementTextureSampler, InVertex.TextureCoordinates.xy).a);
-	OutColor.a *= GetOpacityFromSignedDistance(SampledDistance, InVertex.TextureCoordinates.xy);
-
-	return OutColor;
-}
-
-float4 GetMsdfFontElementColor(VertexOut InVertex)
-{
-	float4 OutColor = InVertex.Color;
-
-	const float4 MultiDistance = float4(ElementTexture.Sample(ElementTextureSampler, InVertex.TextureCoordinates.xy));
-	const float SampledDistance = max(min(MultiDistance.r, MultiDistance.g), min(max(MultiDistance.r, MultiDistance.g), MultiDistance.b));
-	OutColor.a *= GetOpacityFromSignedDistance(SampledDistance, InVertex.TextureCoordinates.xy);
-
-	return OutColor;
-}
-
 float4 GetColor( VertexOut InVertex, float2 UV )
 {
 	float4 FinalColor;
@@ -251,14 +220,6 @@ float4 Main( VertexOut InVertex ) : SV_Target
 	{
 		OutColor = GetColorFontElementColor(InVertex);
 	}
-	else if (ShaderType == ESlateShader::SdfFont)
-	{
-		OutColor = GetSdfFontElementColor(InVertex);
-	}
-	else if (ShaderType == ESlateShader::MsdfFont)
-	{
-		OutColor = GetMsdfFontElementColor(InVertex);
-	}
 	else
 	{
 		OutColor = GetLineSegmentElementColor( InVertex );
diff --git a/Engine/Shaders/StandaloneRenderer/OpenGL/SlateElementPixelShader.glsl b/Engine/Shaders/StandaloneRenderer/OpenGL/SlateElementPixelShader.glsl
--- a/Engine/Shaders/StandaloneRenderer/OpenGL/SlateElementPixelShader.glsl
+++ b/Engine/Shaders/StandaloneRenderer/OpenGL/SlateElementPixelShader.glsl
@@ -19,8 +19,6 @@ precision highp float;
 #define ST_ColorFont		3
 #define ST_Line				4
 #define ST_RoundedBox       7
-#define ST_SdfFont          8
-#define ST_MsdfFont         9
 
 #define USE_LEGACY_DISABLED_EFFECT 0
 
@@ -165,35 +163,6 @@ vec4 GetColorFontElementColor()
 	return OutColor;
 }
 
-float GetOpacityFromSignedDistance(float SampledDistance)
-{
-	vec2 ScreenSpaceTextureSize = vec2(1.0, 1.0) / fwidth(TexCoords.xy);
-	float ScreenSpaceSpread = max(dot(ShaderParams.xy, ScreenSpaceTextureSize), 1.0);
-	float ScreenSpaceDistance = (SampledDistance - ShaderParams.z) * ScreenSpaceSpread;
-	return clamp(ScreenSpaceDistance + 0.5, 0.0, 1.0);
-}
-
-vec4 GetSdfFontElementColor()
-{
-	vec4 OutColor = Color;
-
-	float SampledDistance = texture2D(ElementTexture, TexCoords.xy).a;
-	OutColor.a *= GetOpacityFromSignedDistance(SampledDistance);
-
-	return OutColor;
-}
-
-vec4 GetMsdfFontElementColor()
-{
-	vec4 OutColor = Color;
-
-	vec4 MultiDistance = texture2D(ElementTexture, TexCoords.xy);
-	float SampledDistance = max(min(MultiDistance.r, MultiDistance.g), min(max(MultiDistance.r, MultiDistance.g), MultiDistance.b));
-	OutColor.a *= GetOpacityFromSignedDistance(SampledDistance);
-
-	return OutColor;
-}
-
 vec4 GetDefaultElementColor()
 {
 	vec4 OutColor = Color;
@@ -354,14 +323,6 @@ void main()
 	{
 		OutColor = GetColorFontElementColor();
 	}
-	else if( ShaderType == ST_SdfFont )
-	{
-		OutColor = GetSdfFontElementColor();
-	}
-	else if( ShaderType == ST_MsdfFont )
-	{
-		OutColor = GetMsdfFontElementColor();
-	}
 	else
 	{
 		OutColor = GetLineSegmentElementColor();
