PR #11088: Slate: Fix custom element shader for StandaloneRenderer

diff --git a/Engine/Shaders/StandaloneRenderer/D3D/SlateElementPixelShader.hlsl b/Engine/Shaders/StandaloneRenderer/D3D/SlateElementPixelShader.hlsl
--- a/Engine/Shaders/StandaloneRenderer/D3D/SlateElementPixelShader.hlsl
+++ b/Engine/Shaders/StandaloneRenderer/D3D/SlateElementPixelShader.hlsl
@@ -8,6 +8,7 @@
 #define ESlateShader::GrayscaleFont	2
 #define ESlateShader::ColorFont		3
 #define ESlateShader::LineSegment	4
+#define ESlateShader::Custom		5
 #define ESlateShader::RoundedBox    7
 
 #define USE_LEGACY_DISABLED_EFFECT 0
@@ -200,7 +201,7 @@ float4 Main( VertexOut InVertex ) : SV_Target
 {
 	float4 OutColor;
 
-	if( ShaderType == ESlateShader::Default )
+	if( ShaderType == ESlateShader::Default || ShaderType == ESlateShader::Custom )
 	{
 		OutColor = GetDefaultElementColor( InVertex );
 	}
diff --git a/Engine/Shaders/StandaloneRenderer/OpenGL/SlateElementPixelShader.glsl b/Engine/Shaders/StandaloneRenderer/OpenGL/SlateElementPixelShader.glsl
--- a/Engine/Shaders/StandaloneRenderer/OpenGL/SlateElementPixelShader.glsl
+++ b/Engine/Shaders/StandaloneRenderer/OpenGL/SlateElementPixelShader.glsl
@@ -18,6 +18,7 @@ precision highp float;
 #define ST_GrayscaleFont	2
 #define ST_ColorFont		3
 #define ST_Line				4
+#define ST_Custom			5
 #define ST_RoundedBox       7
 
 #define USE_LEGACY_DISABLED_EFFECT 0
@@ -303,7 +304,7 @@ void main()
 {
 	vec4 OutColor;
 
-	if( ShaderType == ST_Default )
+	if( ShaderType == ST_Default || ShaderType == ST_Custom )
 	{
 		OutColor = GetDefaultElementColor();
 	}
