Slate: Fixed glsl syntax errors on mac

diff --git a/Engine/Shaders/StandaloneRenderer/OpenGL/SlateElementPixelShader.glsl b/Engine/Shaders/StandaloneRenderer/OpenGL/SlateElementPixelShader.glsl
--- a/Engine/Shaders/StandaloneRenderer/OpenGL/SlateElementPixelShader.glsl
+++ b/Engine/Shaders/StandaloneRenderer/OpenGL/SlateElementPixelShader.glsl
@@ -332,8 +332,8 @@ vec4 GetLineSegmentElementColor()
 	vec2 DashParams = TexCoords.zw;
 
 	// Get coverage based on distance from segment sides and ends
-	vec2 OutsideFilterUV = float2(1.0, 1.0);
-	vec2 InsideFilterUV = float2(ShaderParams.x, 0.0);
+	vec2 OutsideFilterUV = vec2(1.0, 1.0);
+	vec2 InsideFilterUV = vec2(ShaderParams.x, 0.0);
 	vec2 LineCoverage = smoothstep(OutsideFilterUV, InsideFilterUV, abs(Gradient));
 
 	float DashLength = DashParams.y;
